Implemented LinuxDeviceProfileSelector (OR-6888).

- Servers can't/shouldn't use ExampleDeviceProfileSelector anyway.
- Can be useful for fine-grained settings on clients (e.g. detect bad drivers and downscale like Android does).

(Merging CL 2703510 to the main branch).

diff --git a/Engine/Config/Linux/LinuxEngine.ini b/Engine/Config/Linux/LinuxEngine.ini
--- a/Engine/Config/Linux/LinuxEngine.ini
+++ b/Engine/Config/Linux/LinuxEngine.ini
@@ -12,3 +12,7 @@ PreferredAccessor=NullSourceCodeAccessor
 [TextureStreaming]
 ; PoolSizeVRAMPercentage is how much percentage of GPU Dedicated VRAM should be used as a TexturePool cache for streaming textures (0 - unlimited streaming)
 PoolSizeVRAMPercentage=70
+
+[DeviceProfileManager]
+DeviceProfileSelectionModule="LinuxDeviceProfileSelector"
+
diff --git a/Engine/Plugins/Runtime/LinuxDeviceProfileSelector/LinuxDeviceProfileSelector.uplugin b/Engine/Plugins/Runtime/LinuxDeviceProfileSelector/LinuxDeviceProfileSelector.uplugin
new file mode 100644
--- /dev/null
+++ b/Engine/Plugins/Runtime/LinuxDeviceProfileSelector/LinuxDeviceProfileSelector.uplugin
@@ -0,0 +1,23 @@
+{
+	"FileVersion" : 3,
+	
+	"FriendlyName" : "Linux Device Profile Selector",
+	"Version" : 1,
+	"VersionName" : "1.0",
+	"CreatedBy" : "Epic Games, Inc.",
+	"CreatedByURL" : "http://epicgames.com",
+	"EngineVersion" : "4.2.0",
+	"Description" : "Linux Device Profile Selector",
+	"Category" : "Device Profile Selectors",
+	"EnabledByDefault" : true,
+	
+	"Modules" :
+	[
+		{
+			"Name" : "LinuxDeviceProfileSelector",
+			"Type" : "RuntimeNoCommandlet",
+			"LoadingPhase" : "PostConfigInit",
+			"WhitelistPlatforms" : [ "Linux" ]
+		}
+	]
+}
diff --git a/Engine/Plugins/Runtime/LinuxDeviceProfileSelector/Source/LinuxDeviceProfileSelector/LinuxDeviceProfileSelector.Build.cs b/Engine/Plugins/Runtime/LinuxDeviceProfileSelector/Source/LinuxDeviceProfileSelector/LinuxDeviceProfileSelector.Build.cs
new file mode 100644
--- /dev/null
+++ b/Engine/Plugins/Runtime/LinuxDeviceProfileSelector/Source/LinuxDeviceProfileSelector/LinuxDeviceProfileSelector.Build.cs
@@ -0,0 +1,38 @@
+// Copyright 1998-2015 Epic Games, Inc. All Rights Reserved.
+
+namespace UnrealBuildTool.Rules
+{
+    public class LinuxDeviceProfileSelector : ModuleRules
+	{
+        public LinuxDeviceProfileSelector(TargetInfo Target)
+		{
+			PublicIncludePaths.AddRange(
+				new string[] {
+                    "Runtime/LinuxDeviceProfileSelector/Public",
+				}
+				);
+
+			PrivateIncludePaths.AddRange(
+				new string[] {
+					"Runtime/LinuxDeviceProfileSelector/Private",
+				}
+				);
+
+			PublicDependencyModuleNames.AddRange(
+				new string[]
+				{
+					"Core",
+				}
+				);
+
+			PrivateDependencyModuleNames.AddRange(
+				new string[]
+				{
+				    "Core",
+				    "CoreUObject",
+				    "Engine",
+				}
+				);
+		}
+	}
+}
diff --git a/Engine/Plugins/Runtime/LinuxDeviceProfileSelector/Source/LinuxDeviceProfileSelector/Private/LinuxDeviceProfileSelectorModule.cpp b/Engine/Plugins/Runtime/LinuxDeviceProfileSelector/Source/LinuxDeviceProfileSelector/Private/LinuxDeviceProfileSelectorModule.cpp
new file mode 100644
--- /dev/null
+++ b/Engine/Plugins/Runtime/LinuxDeviceProfileSelector/Source/LinuxDeviceProfileSelector/Private/LinuxDeviceProfileSelectorModule.cpp
@@ -0,0 +1,25 @@
+// Copyright 1998-2015 Epic Games, Inc. All Rights Reserved.
+
+#include "LinuxDeviceProfileSelectorPrivatePCH.h"
+
+IMPLEMENT_MODULE(FLinuxDeviceProfileSelectorModule, LinuxDeviceProfileSelector);
+
+
+void FLinuxDeviceProfileSelectorModule::StartupModule()
+{
+}
+
+
+void FLinuxDeviceProfileSelectorModule::ShutdownModule()
+{
+}
+
+
+FString const FLinuxDeviceProfileSelectorModule::GetRuntimeDeviceProfileName()
+{
+	// [RCL] 2015-09-22 FIXME: support different environments
+	FString ProfileName = FPlatformProperties::PlatformName();
+	UE_LOG(LogLinux, Log, TEXT("Selected Device Profile: [%s]"), *ProfileName);
+	return ProfileName;
+}
+
diff --git a/Engine/Plugins/Runtime/LinuxDeviceProfileSelector/Source/LinuxDeviceProfileSelector/Private/LinuxDeviceProfileSelectorModule.h b/Engine/Plugins/Runtime/LinuxDeviceProfileSelector/Source/LinuxDeviceProfileSelector/Private/LinuxDeviceProfileSelectorModule.h
new file mode 100644
--- /dev/null
+++ b/Engine/Plugins/Runtime/LinuxDeviceProfileSelector/Source/LinuxDeviceProfileSelector/Private/LinuxDeviceProfileSelectorModule.h
@@ -0,0 +1,31 @@
+// Copyright 1998-2015 Epic Games, Inc. All Rights Reserved.
+
+#pragma once
+
+
+/**
+ * Implements the Linux Device Profile Selector module.
+ */
+class FLinuxDeviceProfileSelectorModule
+	: public IDeviceProfileSelectorModule
+{
+public:
+
+	// Begin IDeviceProfileSelectorModule interface
+	virtual const FString GetRuntimeDeviceProfileName() override;
+	// End IDeviceProfileSelectorModule interface
+
+
+	// Begin IModuleInterface interface
+	virtual void StartupModule() override;
+	virtual void ShutdownModule() override;
+	// End IModuleInterface interface
+
+	
+	/**
+	 * Virtual destructor.
+	 */
+	virtual ~FLinuxDeviceProfileSelectorModule()
+	{
+	}
+};
\ No newline at end of file
diff --git a/Engine/Plugins/Runtime/LinuxDeviceProfileSelector/Source/LinuxDeviceProfileSelector/Private/LinuxDeviceProfileSelectorPrivatePCH.h b/Engine/Plugins/Runtime/LinuxDeviceProfileSelector/Source/LinuxDeviceProfileSelector/Private/LinuxDeviceProfileSelectorPrivatePCH.h
new file mode 100644
--- /dev/null
+++ b/Engine/Plugins/Runtime/LinuxDeviceProfileSelector/Source/LinuxDeviceProfileSelector/Private/LinuxDeviceProfileSelectorPrivatePCH.h
@@ -0,0 +1,20 @@
+// Copyright 1998-2015 Epic Games, Inc. All Rights Reserved.
+
+#pragma once
+
+#include "LinuxDeviceProfileSelector.h"
+
+
+/* Private dependencies
+*****************************************************************************/
+
+#include "Core.h"
+#include "Engine.h"
+#include "ModuleManager.h"
+#include "ModuleInterface.h"
+
+
+/* Private includes
+ *****************************************************************************/
+
+#include "LinuxDeviceProfileSelectorModule.h"
diff --git a/Engine/Plugins/Runtime/LinuxDeviceProfileSelector/Source/LinuxDeviceProfileSelector/Public/LinuxDeviceProfileSelector.h b/Engine/Plugins/Runtime/LinuxDeviceProfileSelector/Source/LinuxDeviceProfileSelector/Public/LinuxDeviceProfileSelector.h
new file mode 100644
--- /dev/null
+++ b/Engine/Plugins/Runtime/LinuxDeviceProfileSelector/Source/LinuxDeviceProfileSelector/Public/LinuxDeviceProfileSelector.h
@@ -0,0 +1,8 @@
+// Copyright 1998-2015 Epic Games, Inc. All Rights Reserved.
+
+#pragma once
+
+/* Dependencies
+*****************************************************************************/
+#include "ModuleManager.h"
+#include "IDeviceProfileSelectorModule.h"
