Virtual Asset Fixes for CookedEditor targets
- Set VA.DisableVirtualization CVar for CookedEditor targets
- Use UE_VIRTUALIZATION_CONNECTION_LAZY_INIT for [redacted] CookedEditor targets
- setup FLinuxDeviceProfileSelectorModule so Linux CookedEditor use the right device profile
- Fix bLazyInitConnections not being set by UE_VIRTUALIZATION_CONNECTION_LAZY_INIT and fix logging
- Add AllowVirtualAssets to automation scripts for [redacted]

diff --git a/Engine/Config/BaseDeviceProfiles.ini b/Engine/Config/BaseDeviceProfiles.ini
--- a/Engine/Config/BaseDeviceProfiles.ini
+++ b/Engine/Config/BaseDeviceProfiles.ini
@@ -165,8 +165,7 @@ DeviceType=Windows
 BaseProfileName=Windows
 +CVars=s.AllowUnversionedContentInEditor=1
 +CVars=cook.AllowCookedDataInEditorBuilds=1
-
-
++CVars=VA.DisableVirtualization=1
 
 ;----------------
 ; iOS Devices
@@ -948,6 +947,13 @@ BaseProfileName=Linux
 MeshLODSettings=
 TextureLODSettings=
 
+[LinuxCookedEditor DeviceProfile]
+DeviceType=Linux
+BaseProfileName=Linux
++CVars=s.AllowUnversionedContentInEditor=1
++CVars=cook.AllowCookedDataInEditorBuilds=1
++CVars=VA.DisableVirtualization=1
+
 [HoloLens DeviceProfile]
 DeviceType=HoloLens
 BaseProfileName=Windows
diff --git a/Engine/Plugins/Runtime/LinuxDeviceProfileSelector/Source/LinuxDeviceProfileSelector/Private/LinuxDeviceProfileSelectorModule.cpp b/Engine/Plugins/Runtime/LinuxDeviceProfileSelector/Source/LinuxDeviceProfileSelector/Private/LinuxDeviceProfileSelectorModule.cpp
--- a/Engine/Plugins/Runtime/LinuxDeviceProfileSelector/Source/LinuxDeviceProfileSelector/Private/LinuxDeviceProfileSelectorModule.cpp
+++ b/Engine/Plugins/Runtime/LinuxDeviceProfileSelector/Source/LinuxDeviceProfileSelector/Private/LinuxDeviceProfileSelectorModule.cpp
@@ -18,8 +18,14 @@ void FLinuxDeviceProfileSelectorModule::ShutdownModule()
 
 FString const FLinuxDeviceProfileSelectorModule::GetRuntimeDeviceProfileName()
 {
+#if UE_IS_COOKED_EDITOR
+	// some heuristics to determine a cooked editor
+	FString ProfileName = TEXT("LinuxCookedEditor");
+#else
 	// [RCL] 2015-09-22 FIXME: support different environments
 	FString ProfileName = FPlatformProperties::PlatformName();
+#endif
+
 	UE_LOG(LogLinux, Log, TEXT("Selected Device Profile: [%s]"), *ProfileName);
 	return ProfileName;
 }
diff --git a/Engine/Source/Developer/Virtualization/Private/VirtualizationManager.cpp b/Engine/Source/Developer/Virtualization/Private/VirtualizationManager.cpp
--- a/Engine/Source/Developer/Virtualization/Private/VirtualizationManager.cpp
+++ b/Engine/Source/Developer/Virtualization/Private/VirtualizationManager.cpp
@@ -960,7 +960,8 @@ void FVirtualizationManager::ApplySettingsFromConfigFiles(const FConfigFile& Con
 		UE_LOG(LogVirtualization, Error, TEXT("Failed to load [Core.VirtualizationModule].LazyInitConnections from config file!"));
 	}
 #else
-	UE_LOG(LogVirtualization, Display, TEXT("\tLazyInitConnections : false (set by code)"), bLazyInitConnections ? TEXT("true") : TEXT("false"));
+	bLazyInitConnections = true;
+	UE_LOG(LogVirtualization, Display, TEXT("\tLazyInitConnections : %s (set by code)"), bLazyInitConnections ? TEXT("true") : TEXT("false"));
 #endif //UE_VIRTUALIZATION_CONNECTION_LAZY_INIT
 
 	// Check for any legacy settings and print them out (easier to do this in one block rather than one and time)
