PR #13368: Adding support for XR device profile selection on Linux

#rnx
#Synced-CL 51211712

diff --git a/Engine/Plugins/Runtime/LinuxDeviceProfileSelector/Source/LinuxDeviceProfileSelector/LinuxDeviceProfileSelector.Build.cs b/Engine/Plugins/Runtime/LinuxDeviceProfileSelector/Source/LinuxDeviceProfileSelector/LinuxDeviceProfileSelector.Build.cs
--- a/Engine/Plugins/Runtime/LinuxDeviceProfileSelector/Source/LinuxDeviceProfileSelector/LinuxDeviceProfileSelector.Build.cs
+++ b/Engine/Plugins/Runtime/LinuxDeviceProfileSelector/Source/LinuxDeviceProfileSelector/LinuxDeviceProfileSelector.Build.cs
@@ -7,20 +7,19 @@ namespace UnrealBuildTool.Rules
         public LinuxDeviceProfileSelector(ReadOnlyTargetRules Target) : base(Target)
 		{
 			PublicDependencyModuleNames.AddRange(
-				new string[]
-				{
+				new string[] {
 					"Core",
 				}
-				);
+			);
 
 			PrivateDependencyModuleNames.AddRange(
-				new string[]
-				{
+				new string[] {
 				    "Core",
 				    "CoreUObject",
 				    "Engine",
+					"HeadMountedDisplay",
 				}
-				);
+			);
 		}
 	}
 }
diff --git a/Engine/Plugins/Runtime/LinuxDeviceProfileSelector/Source/LinuxDeviceProfileSelector/Private/LinuxDeviceProfileSelectorModule.cpp b/Engine/Plugins/Runtime/LinuxDeviceProfileSelector/Source/LinuxDeviceProfileSelector/Private/LinuxDeviceProfileSelectorModule.cpp
--- a/Engine/Plugins/Runtime/LinuxDeviceProfileSelector/Source/LinuxDeviceProfileSelector/Private/LinuxDeviceProfileSelectorModule.cpp
+++ b/Engine/Plugins/Runtime/LinuxDeviceProfileSelector/Source/LinuxDeviceProfileSelector/Private/LinuxDeviceProfileSelectorModule.cpp
@@ -2,6 +2,8 @@
 
 #include "LinuxDeviceProfileSelectorModule.h"
 #include "Modules/ModuleManager.h"
+#include "IHeadMountedDisplayModule.h"
+#include "StereoRendering.h"
 
 IMPLEMENT_MODULE(FLinuxDeviceProfileSelectorModule, LinuxDeviceProfileSelector);
 
@@ -18,13 +20,47 @@ void FLinuxDeviceProfileSelectorModule::ShutdownModule()
 
 FString const FLinuxDeviceProfileSelectorModule::GetRuntimeDeviceProfileName()
 {
+	static FString ProfileName;
+
+	if (ProfileName.IsEmpty())
+	{
 #if UE_IS_COOKED_EDITOR
-	// some heuristics to determine a cooked editor
-	FString ProfileName = TEXT("LinuxCookedEditor");
+		// some heuristics to determine a cooked editor
+		ProfileName = TEXT("LinuxCookedEditor");
 #else
-	// [RCL] 2015-09-22 FIXME: support different environments
-	FString ProfileName = FPlatformProperties::PlatformName();
+		// [RCL] 2015-09-22 FIXME: support different environments
+		ProfileName = FPlatformProperties::PlatformName();
+
+		if (!IsRunningDedicatedServer() && IStereoRendering::IsStartInVR() && IHeadMountedDisplayModule::IsAvailable())
+		{
+				if (!GIsEditor || IHeadMountedDisplayModule::Get().IsStandaloneStereoOnlyDevice())
+				{
+					// Add support for proper profile identification on Linux platforms that support XR
+					FString HMDRequestedProfileName = IHeadMountedDisplayModule::Get().GetDeviceSystemName();
+
+					if (!HMDRequestedProfileName.IsEmpty())
+					{
+						FString RuntimeName, DeviceDriverName;
+
+						// if DeviceSystemName includes a "Runtime:Device" use the Device part.  Otherwise just use the whole string
+						if (HMDRequestedProfileName.Split(TEXT(":"), &RuntimeName, &DeviceDriverName))
+						{
+							DeviceDriverName.TrimStartAndEndInline();
+
+							if (!DeviceDriverName.IsEmpty())
+							{
+								ProfileName = DeviceDriverName;
+							}
+						}
+						else
+						{
+							ProfileName = HMDRequestedProfileName;
+						}
+					}
+				}
+		}
 #endif
+	}
 
 	UE_LOG(LogLinux, Log, TEXT("Selected Device Profile: [%s]"), *ProfileName);
 	return ProfileName;
